Non-unity build fixes.

diff --git a/Engine/Plugins/MovieScene/ActorSequenceEditor/Source/ActorSequenceEditor/Private/ActorSequenceComponentCustomization.cpp b/Engine/Plugins/MovieScene/ActorSequenceEditor/Source/ActorSequenceEditor/Private/ActorSequenceComponentCustomization.cpp
--- a/Engine/Plugins/MovieScene/ActorSequenceEditor/Source/ActorSequenceEditor/Private/ActorSequenceComponentCustomization.cpp
+++ b/Engine/Plugins/MovieScene/ActorSequenceEditor/Source/ActorSequenceEditor/Private/ActorSequenceComponentCustomization.cpp
@@ -16,6 +16,7 @@
 #include "ScopedTransaction.h"
 #include "ISequencerModule.h"
 #include "ActorSequenceEditorTabSummoner.h"
+#include "Widgets/Input/SButton.h"
 
 #define LOCTEXT_NAMESPACE "ActorSequenceComponentCustomization"
 
diff --git a/Engine/Plugins/MovieScene/ActorSequenceEditor/Source/ActorSequenceEditor/Private/ActorSequenceComponentCustomization.h b/Engine/Plugins/MovieScene/ActorSequenceEditor/Source/ActorSequenceEditor/Private/ActorSequenceComponentCustomization.h
--- a/Engine/Plugins/MovieScene/ActorSequenceEditor/Source/ActorSequenceEditor/Private/ActorSequenceComponentCustomization.h
+++ b/Engine/Plugins/MovieScene/ActorSequenceEditor/Source/ActorSequenceEditor/Private/ActorSequenceComponentCustomization.h
@@ -2,7 +2,12 @@
 
 #pragma once
 
+#include "CoreMinimal.h"
 #include "IDetailCustomization.h"
+#include "Widgets/Layout/SBox.h"
+#include "Input/Reply.h"
+#include "UObject/WeakObjectPtr.h"
+#include "Framework/Docking/TabManager.h"
 
 class IDetailsView;
 class UActorSequence;
diff --git a/Engine/Plugins/MovieScene/ActorSequenceEditor/Source/ActorSequenceEditor/Private/ActorSequenceEditor.cpp b/Engine/Plugins/MovieScene/ActorSequenceEditor/Source/ActorSequenceEditor/Private/ActorSequenceEditor.cpp
--- a/Engine/Plugins/MovieScene/ActorSequenceEditor/Source/ActorSequenceEditor/Private/ActorSequenceEditor.cpp
+++ b/Engine/Plugins/MovieScene/ActorSequenceEditor/Source/ActorSequenceEditor/Private/ActorSequenceEditor.cpp
@@ -13,6 +13,8 @@
 #include "PropertyEditorModule.h"
 #include "Styling/SlateStyle.h"
 #include "WorkflowTabManager.h"
+#include "Modules/ModuleManager.h"
+#include "Widgets/Docking/SDockTab.h"
 
 
 #define LOCTEXT_NAMESPACE "ActorSequenceEditor"
diff --git a/Engine/Plugins/MovieScene/ActorSequenceEditor/Source/ActorSequenceEditor/Private/ActorSequenceEditorTabSummoner.cpp b/Engine/Plugins/MovieScene/ActorSequenceEditor/Source/ActorSequenceEditor/Private/ActorSequenceEditorTabSummoner.cpp
--- a/Engine/Plugins/MovieScene/ActorSequenceEditor/Source/ActorSequenceEditor/Private/ActorSequenceEditorTabSummoner.cpp
+++ b/Engine/Plugins/MovieScene/ActorSequenceEditor/Source/ActorSequenceEditor/Private/ActorSequenceEditorTabSummoner.cpp
@@ -9,6 +9,13 @@
 #include "SSCSEditor.h"
 #include "SlateIconFinder.h"
 #include "BlueprintEditorUtils.h"
+#include "EditorStyleSet.h"
+#include "Widgets/Images/SImage.h"
+#include "Editor.h"
+#include "ScopedTransaction.h"
+#include "Framework/MultiBox/MultiBoxBuilder.h"
+#include "Framework/Application/SlateApplication.h"
+
 
 #define LOCTEXT_NAMESPACE "ActorSequenceEditorSummoner"
 
diff --git a/Engine/Plugins/Online/OnlineSubsystem/Source/Private/OnlineSubsystemImpl.cpp b/Engine/Plugins/Online/OnlineSubsystem/Source/Private/OnlineSubsystemImpl.cpp
--- a/Engine/Plugins/Online/OnlineSubsystem/Source/Private/OnlineSubsystemImpl.cpp
+++ b/Engine/Plugins/Online/OnlineSubsystem/Source/Private/OnlineSubsystemImpl.cpp
@@ -4,6 +4,7 @@
 #include "Containers/Ticker.h"
 #include "Misc/App.h"
 #include "NamedInterfaces.h"
+#include "OnlineError.h"
 #include "Interfaces/OnlineIdentityInterface.h"
 #include "Interfaces/OnlineSessionInterface.h"
 #include "Interfaces/OnlineFriendsInterface.h"
diff --git a/Engine/Source/Developer/AssetTools/Private/AssetTypeActions/AssetTypeActions_AnimationAsset.cpp b/Engine/Source/Developer/AssetTools/Private/AssetTypeActions/AssetTypeActions_AnimationAsset.cpp
--- a/Engine/Source/Developer/AssetTools/Private/AssetTypeActions/AssetTypeActions_AnimationAsset.cpp
+++ b/Engine/Source/Developer/AssetTools/Private/AssetTypeActions/AssetTypeActions_AnimationAsset.cpp
@@ -12,6 +12,7 @@
 #include "SSkeletonWidget.h"
 #include "IAnimationEditorModule.h"
 #include "Preferences/PersonaOptions.h"
+#include "Framework/MultiBox/MultiBoxBuilder.h"
 
 #define LOCTEXT_NAMESPACE "AssetTypeActions"
 
diff --git a/Engine/Source/Developer/BlueprintNativeCodeGen/Private/BlueprintNativeCodeGenManifest.h b/Engine/Source/Developer/BlueprintNativeCodeGen/Private/BlueprintNativeCodeGenManifest.h
--- a/Engine/Source/Developer/BlueprintNativeCodeGen/Private/BlueprintNativeCodeGenManifest.h
+++ b/Engine/Source/Developer/BlueprintNativeCodeGen/Private/BlueprintNativeCodeGenManifest.h
@@ -5,6 +5,7 @@
 #include "CoreMinimal.h"
 #include "UObject/ObjectMacros.h"
 #include "Misc/StringAssetReference.h"
+#include "Engine/Blueprint.h"
 #include "BlueprintNativeCodeGenManifest.generated.h"
 
 class FAssetData;
diff --git a/Engine/Source/Developer/BlueprintNativeCodeGen/Public/BlueprintNativeCodeGenModule.h b/Engine/Source/Developer/BlueprintNativeCodeGen/Public/BlueprintNativeCodeGenModule.h
--- a/Engine/Source/Developer/BlueprintNativeCodeGen/Public/BlueprintNativeCodeGenModule.h
+++ b/Engine/Source/Developer/BlueprintNativeCodeGen/Public/BlueprintNativeCodeGenModule.h
@@ -7,6 +7,7 @@
 #include "UObject/AssetPtr.h"
 #include "Modules/ModuleInterface.h"
 #include "Modules/ModuleManager.h"
+#include "Engine/Blueprint.h"
 
 class UBlueprint;
 enum class ESavePackageResult;
diff --git a/Engine/Source/Editor/VREditor/SVRRadialPanel.cpp b/Engine/Source/Editor/VREditor/SVRRadialPanel.cpp
--- a/Engine/Source/Editor/VREditor/SVRRadialPanel.cpp
+++ b/Engine/Source/Editor/VREditor/SVRRadialPanel.cpp
@@ -3,7 +3,8 @@
 #include "SVRRadialPanel.h"
 #include "VREditorModule.h"
 #include "VREditorUISystem.h"
-
+#include "HAL/IConsoleManager.h"
+#include "Layout/LayoutUtils.h"
 
 namespace VREd
 {
diff --git a/Engine/Source/Editor/VREditor/SVRRadialPanel.h b/Engine/Source/Editor/VREditor/SVRRadialPanel.h
--- a/Engine/Source/Editor/VREditor/SVRRadialPanel.h
+++ b/Engine/Source/Editor/VREditor/SVRRadialPanel.h
@@ -12,6 +12,7 @@
 #include "Widgets/SWidget.h"
 #include "Layout/Children.h"
 #include "Widgets/SPanel.h"
+#include "Widgets/Input/SButton.h"
 
 /**
  * A RadialPanel contains one child and describes how that child should be arranged on the screen.
