Actor Sequence: Fix crash on compiling blueprint with deleted actor sequence

diff --git a/Engine/Plugins/MovieScene/ActorSequenceEditor/Source/ActorSequenceEditor/Private/ActorSequenceEditorTabSummoner.cpp b/Engine/Plugins/MovieScene/ActorSequenceEditor/Source/ActorSequenceEditor/Private/ActorSequenceEditorTabSummoner.cpp
--- a/Engine/Plugins/MovieScene/ActorSequenceEditor/Source/ActorSequenceEditor/Private/ActorSequenceEditorTabSummoner.cpp
+++ b/Engine/Plugins/MovieScene/ActorSequenceEditor/Source/ActorSequenceEditor/Private/ActorSequenceEditorTabSummoner.cpp
@@ -240,7 +240,10 @@ public:
 			}
 			else if (UBlueprintGeneratedClass* GeneratedClass = LocalActorSequence->GetTypedOuter<UBlueprintGeneratedClass>())
 			{
-				return GeneratedClass->SimpleConstructionScript->GetComponentEditorActorInstance();
+				if (GeneratedClass->SimpleConstructionScript)
+				{
+					return GeneratedClass->SimpleConstructionScript->GetComponentEditorActorInstance();
+				}
 			}
 		}
 		
